#UE4 - Enable IOS advertising for all IOS apps

diff --git a/Engine/Config/IOS/IOSEngine.ini b/Engine/Config/IOS/IOSEngine.ini
--- a/Engine/Config/IOS/IOSEngine.ini
+++ b/Engine/Config/IOS/IOSEngine.ini
@@ -4,6 +4,9 @@ AudioDeviceModuleName=IOSAudio
 [OnlineSubsystem]
 DefaultPlatformService=IOS
 
+[Advertising]
+DefaultProviderName=IOSAdvertising
+
 [Plugins]
 +EnabledPlugins="AppleMoviePlayer"
 +EnabledPlugins="IOSDeviceProfileSelector"
