Fixed: Strategy Game - FPS regression on Samsung Galaxy S4 noticed a decrease in frame rate from 4.8 to 4.9

diff --git a/Engine/Config/Android/AndroidEngine.ini b/Engine/Config/Android/AndroidEngine.ini
--- a/Engine/Config/Android/AndroidEngine.ini
+++ b/Engine/Config/Android/AndroidEngine.ini
@@ -11,3 +11,7 @@ DefaultProviderName=AndroidAdvertising
 
 [OnlineSubsystem]
 DefaultPlatformService=GooglePlay
+
+[SlateRenderer]
+NumPreallocatedVertices=200
+
diff --git a/Engine/Config/IOS/IOSEngine.ini b/Engine/Config/IOS/IOSEngine.ini
--- a/Engine/Config/IOS/IOSEngine.ini
+++ b/Engine/Config/IOS/IOSEngine.ini
@@ -13,4 +13,7 @@ DefaultProviderName=IOSAdvertising
 DeviceProfileSelectionModule="IOSDeviceProfileSelector"
 
 [OnlineSubsystemIOS.Store]
-bSupportInAppPurchasing=False
\ No newline at end of file
+bSupportInAppPurchasing=False
+
+[SlateRenderer]
+NumPreallocatedVertices=200
\ No newline at end of file
