- Ensure default behavior is still to use task graph threads

diff --git a/Engine/Config/IOS/BaseIOSEngine.ini b/Engine/Config/IOS/BaseIOSEngine.ini
--- a/Engine/Config/IOS/BaseIOSEngine.ini
+++ b/Engine/Config/IOS/BaseIOSEngine.ini
@@ -3,3 +3,6 @@ NativePlatformService=IOS
 
 [Audio]
 AudioCaptureModuleName=AudioCaptureAudioUnit
+
+[Ping]
+StackSize=0
\ No newline at end of file
