Enable bSupportsByteBufferComputeShaders for SM5 via Mac and iOS DataDrivenPlatformInfo.ini

#rnx

diff --git a/Engine/Config/IOS/DataDrivenPlatformInfo.ini b/Engine/Config/IOS/DataDrivenPlatformInfo.ini
--- a/Engine/Config/IOS/DataDrivenPlatformInfo.ini
+++ b/Engine/Config/IOS/DataDrivenPlatformInfo.ini
@@ -45,4 +45,7 @@ InactiveIconName=LevelEditor.PreviewMode.iOSSM5.Disabled
 MenuText=NSLOCTEXT("PreviewPlatform", "PreviewMenuText_iOSSM5", "iOS (SM5 Renderer)")
 MenuTooltip=NSLOCTEXT("PreviewPlatform", "PreviewMenuTooltip_iOSSM5", "Mobile preview using iOS SM5 material quality settings.")
 IconText=NSLOCTEXT("PreviewPlatform", "PreviewIconText_iOSSM5", "iOS SM5")
-DeviceProfileName=IOS
\ No newline at end of file
+DeviceProfileName=IOS
+
+[ShaderPlatform METAL_SM5]
+bSupportsByteBufferComputeShaders=true
diff --git a/Engine/Config/Mac/DataDrivenPlatformInfo.ini b/Engine/Config/Mac/DataDrivenPlatformInfo.ini
--- a/Engine/Config/Mac/DataDrivenPlatformInfo.ini
+++ b/Engine/Config/Mac/DataDrivenPlatformInfo.ini
@@ -23,3 +23,5 @@ bUsesHostCompiler=true
 bUATClosesAfterLaunch=true
 PlatformGroupName=Desktop
 
+[ShaderPlatform METAL_SM5]
+bSupportsByteBufferComputeShaders=true
diff --git a/Engine/Shaders/Public/Platform/Metal/MetalCommon.ush b/Engine/Shaders/Public/Platform/Metal/MetalCommon.ush
--- a/Engine/Shaders/Public/Platform/Metal/MetalCommon.ush
+++ b/Engine/Shaders/Public/Platform/Metal/MetalCommon.ush
@@ -5,7 +5,7 @@
 =============================================================================*/
 
 // Update this GUID to improve shader recompilation for Metal only shaders
-// GUID = 829A0E2C-A667-4ABF-82A8-E27DA13447D2
+// GUID = 2FCBEC23-BA55-4BCB-9C9D-41C08D506184
 
 #pragma once
 
