Removing bIsMobile from METAL_MRT. This was incorrectly added when moving hardcoded shader platforms into DDPI.

diff --git a/Engine/Config/IOS/DataDrivenPlatformInfo.ini b/Engine/Config/IOS/DataDrivenPlatformInfo.ini
--- a/Engine/Config/IOS/DataDrivenPlatformInfo.ini
+++ b/Engine/Config/IOS/DataDrivenPlatformInfo.ini
@@ -56,7 +56,6 @@ bSupportsManualVertexFetch=false
 [ShaderPlatform METAL_MRT]
 Language=Metal
 MaxFeatureLevel=SM5
-bIsMobile=true
 bIsMetalMRT=true
 bTargetsTiledGPU=true
 bSupportsIndexBufferUAVs=true
