Backing out changes for now so we can fix some global functions at the same time.

[Backout] - CL18340160
Original CL Desc
-----------------------------------------------------------------
Removing bIsMobile from METAL_MRT. This was incorrectly added when moving hardcoded shader platforms into DDPI.

diff --git a/Engine/Config/IOS/DataDrivenPlatformInfo.ini b/Engine/Config/IOS/DataDrivenPlatformInfo.ini
--- a/Engine/Config/IOS/DataDrivenPlatformInfo.ini
+++ b/Engine/Config/IOS/DataDrivenPlatformInfo.ini
@@ -56,6 +56,7 @@ bSupportsManualVertexFetch=false
 [ShaderPlatform METAL_MRT]
 Language=Metal
 MaxFeatureLevel=SM5
+bIsMobile=true
 bIsMetalMRT=true
 bTargetsTiledGPU=true
 bSupportsIndexBufferUAVs=true
