Fix incorrect ShaderFormat name in IOS DDPI causing issues when IOStore was enabled

diff --git a/Engine/Config/IOS/DataDrivenPlatformInfo.ini b/Engine/Config/IOS/DataDrivenPlatformInfo.ini
--- a/Engine/Config/IOS/DataDrivenPlatformInfo.ini
+++ b/Engine/Config/IOS/DataDrivenPlatformInfo.ini
@@ -48,7 +48,7 @@ DeviceProfileName=IOS
 [ShaderPlatform METAL]
 Language=Metal
 MaxFeatureLevel=ES3_1
-ShaderFormat=SF_METAL_MACES3_1
+ShaderFormat=SF_METAL
 bIsMobile=true
 bTargetsTiledGPU=true
 bSupports4ComponentUAVReadWrite=true
