Fix unquoted string in iOS DeviceProfile

#rnx

diff --git a/Engine/Config/IOS/DataDrivenPlatformInfo.ini b/Engine/Config/IOS/DataDrivenPlatformInfo.ini
--- a/Engine/Config/IOS/DataDrivenPlatformInfo.ini
+++ b/Engine/Config/IOS/DataDrivenPlatformInfo.ini
@@ -33,7 +33,7 @@ InactiveIconName=LevelEditor.PreviewMode.iOS.Disabled
 MenuTooltip=NSLOCTEXT("PreviewPlatform", "PreviewMenuTooltip_iOS", "Mobile preview using iOS material quality settings.")
 IconText=NSLOCTEXT("PreviewPlatform", "PreviewIconText_iOS", "iOS")
 DeviceProfileName=IOS,iPhone6S,iPadPro6_129
-FriendlyName=Generic IOS,iPhone 6S, iPad Pro 6 12.9"
+FriendlyName="Generic IOS","iPhone 6S","iPad Pro 6 12.9"
 
 [PreviewPlatform IOSMetalSM5]
 PlatformName=IOS
