Move settings from iOSEngine.ini to BaseiOSEngine.ini, also correcting a mistake from a previous CL

diff --git a/Engine/Config/IOS/BaseIOSEngine.ini b/Engine/Config/IOS/BaseIOSEngine.ini
--- a/Engine/Config/IOS/BaseIOSEngine.ini
+++ b/Engine/Config/IOS/BaseIOSEngine.ini
@@ -1,3 +1,74 @@
+[/Script/Engine.Engine]
+; DefaultBloomKernel is not used on mobile
+DefaultBloomKernelTextureName=/Engine/EngineResources/DefaultTexture.DefaultTexture
+
+; Default film grain is not used on mobile
+DefaultFilmGrainTextureName=""
+
+[/Script/Engine.GarbageCollectionSettings]
+gc.MaxObjectsInGame=131072
+
+[/Script/Engine.RendererSettings]
+r.DistanceFields=0
+r.TemporalAA.Mobile.UseCompute=0
+r.Substrate.RoughDiffuse=0
+r.Substrate.ShadingQuality=2
+r.Substrate.TileCoord8bits=1
+r.Substrate.SheenQuality=2
+r.Substrate.Glints=0
+r.Substrate.SpecularProfile=0
+r.Substrate.ClosuresPerPixel=1
+; Enable GPU culling for landscape
+landscape.SupportGPUCulling=1
+
+[Advertising]
+DefaultProviderName=IOSAdvertising
+
+[AlternateTextureCompression]
+;disabling TextureCompressionFormat for Android/IOS no longer needed
+;TextureCompressionFormat=""
+;TextureFormatPrefix=""
+
+[Audio]
+AudioCaptureModuleName=AudioCaptureAudioUnit
+; Defining below allows switching to audio mixer using -audiomixer commandline
+AudioMixerModuleName=AudioMixerAudioUnit
+
+; Defines a platform-specific volume headroom (in dB) for audio to provide better platform consistency with respect to volume levels.
+PlatformHeadroomDB=-6
+
+PlatformFormat=ADPCM
+PlatformStreamingFormat=ADPCM
+
+[BackgroundHttp.iOSSettings]
+;While the app is in the forground, time out if we don't get a response in this timeframe
+ActiveReceiveTimeout=30
+;While the app is in the background, time out if we don't get a response in this timeframe
+BackgroundReceiveTimeout=120
+;Regardless of app state, if the download hasn't finished in this long, time out
+BackgroundHttpResourceTimeout=3600
+;If retry data exists for a task, how many times we should try using the Retry data before moving onto a different CDN. ( -1 = unlimited )
+RetryResumeDataLimit=3;
+
+[ConsoleVariables]
+Slate.CacheRenderData=0
+r.HairStrands.Simulation=0
+r.HairStrands.Strands=0
+r.HairStrands.Binding=1
+
+[DeviceProfileManager]
+DeviceProfileSelectionModule="IOSDeviceProfileSelector"
+
+[LocalNotification]
+DefaultPlatformService=IOSLocalNotification
+
+[MemoryMappedFiles]
+Enable=true
+Alignment=16384
+
+[OnlineServices.EOS.Auth]
+DefaultExternalCredentialTypeStr=Apple
+
 [OnlineSubsystem]
 DefaultPlatformService=IOS
 NativePlatformService=IOS
@@ -6,11 +77,35 @@ LocalPlatformName=IOS
 [OnlineSubsystemIOS.Store]
 bSupportsInAppPurchasing=true
 
-[Audio]
-AudioCaptureModuleName=AudioCaptureAudioUnit
+[PlatformCrypto]
+PlatformRequiresDataCrypto=True
+PakSigningRequired=False
+
+[PlatformMemoryBuckets]
+DefaultMemoryBucket_MinGB=4
+SmallerMemoryBucket_MinGB=3
+SmallestMemoryBucket_MinGB=2
+
+[Plugins]
++EnabledPlugins=Crashlytics
 
 [Ping]
 StackSize=0
 
-[OnlineServices.EOS.Auth]
-DefaultExternalCredentialTypeStr=Apple
\ No newline at end of file
+[SlateRenderer]
+NumPreallocatedVertices=200
+
+[SystemSettings]
+; iOS should always have IPv6 support on
+net.DisableIPv6=0
+;Perplatform to PerQualityLevel conversion mapping for platform groups
+QualityLevelMapping="low"
+
+;Manual range settings for Niagara quality levels. Simplifies handling of device profiles/fragments and what quality levels they can set.
+;Please ensure device profiles or fragments for this platform do not set fx.Niagara.QualityLevel values outside of this range
+fx.Niagara.QualityLevel.Min=0
+fx.Niagara.QualityLevel.Max=1
+
+[TextureStreaming]
+; PoolSizeVRAMPercentage is how much percentage of GPU Dedicated VRAM should be used as a TexturePool cache for streaming textures (0 - unlimited streaming)
+PoolSizeVRAMPercentage=70
diff --git a/Engine/Config/IOS/IOSEngine.ini b/Engine/Config/IOS/IOSEngine.ini
deleted file mode 100644
diff --git a/Engine/Config/Mac/BaseMacEngine.ini b/Engine/Config/Mac/BaseMacEngine.ini
--- a/Engine/Config/Mac/BaseMacEngine.ini
+++ b/Engine/Config/Mac/BaseMacEngine.ini
@@ -1,5 +1,6 @@
 [Audio]
 ; Defining below allows switching to audio mixer using -audiomixer commandline
+AudioCaptureModuleName=AudioCaptureRtAudio
 AudioMixerModuleName=AudioMixerCoreAudio
 PlatformFormat=OGG
 PlatformStreamingFormat=OGG
