Set r.Forward.LightGridSizeZ = 8, r.Forward.MaxCulledLightsPerCell = 8 for iOS

diff --git a/Engine/Config/IOS/BaseIOSEngine.ini b/Engine/Config/IOS/BaseIOSEngine.ini
--- a/Engine/Config/IOS/BaseIOSEngine.ini
+++ b/Engine/Config/IOS/BaseIOSEngine.ini
@@ -53,6 +53,8 @@ Slate.CacheRenderData=0
 r.HairStrands.Simulation=0
 r.HairStrands.Strands=0
 r.HairStrands.Binding=0
+r.Forward.LightGridSizeZ=8
+r.Forward.MaxCulledLightsPerCell=8
 
 [DeviceProfileManager]
 DeviceProfileSelectionModule="IOSDeviceProfileSelector"
