Disable parallel GDME on IOS.

It was disabled by default and is explicitly disabled on other platforms.
However, it was recently enabled by default, affecting IOS.

#rnx

diff --git a/Engine/Config/IOS/BaseIOSEngine.ini b/Engine/Config/IOS/BaseIOSEngine.ini
--- a/Engine/Config/IOS/BaseIOSEngine.ini
+++ b/Engine/Config/IOS/BaseIOSEngine.ini
@@ -18,6 +18,8 @@ r.Substrate.SheenQuality=2
 r.Substrate.Glints=0
 r.Substrate.SpecularProfile=0
 r.Substrate.ClosuresPerPixel=1
+; Disables parallel GDME by default due to not enough threads
+r.Visibility.DynamicMeshElements.Parallel=0
 
 [Advertising]
 DefaultProviderName=IOSAdvertising
@@ -129,4 +131,4 @@ fx.Niagara.QualityLevel.Max=1
 PoolSizeVRAMPercentage=70
 
 [BuildPatchTool]
-VerificationThreadCount=3
\ No newline at end of file
+VerificationThreadCount=3
