[Deformer Graph & ComputeFramework] Enabled IOS & Switch platform support

diff --git a/Engine/Config/IOS/DataDrivenPlatformInfo.ini b/Engine/Config/IOS/DataDrivenPlatformInfo.ini
--- a/Engine/Config/IOS/DataDrivenPlatformInfo.ini
+++ b/Engine/Config/IOS/DataDrivenPlatformInfo.ini
@@ -57,6 +57,7 @@ bSupportsManualVertexFetch=false
 bSupportsVolumeTextureAtomics=false
 bSupportsShaderPipelines = false
 bSupportsIndependentSamplers=true
+bSupportsComputeFramework = true
 FriendlyName=LOCTEXT("FriendlyShaderPlatformName_IOSMobile", "IOS Metal Mobile")
 
 [ShaderPlatform METAL_MRT]
@@ -71,6 +72,7 @@ bSupportsManualVertexFetch=false
 bSupportsVolumeTextureAtomics=false
 bSupportsDistanceFields=true
 bSupportsIndependentSamplers=true
+bSupportsComputeFramework = true
 FriendlyName=LOCTEXT("FriendlyShaderPlatformName_IOSSM5", "IOS Metal SM5")
 
 [ShaderPlatform METAL_SIM]
@@ -84,4 +86,5 @@ bSupportsManualVertexFetch=false
 bSupportsVolumeTextureAtomics=false
 bSupportsShaderPipelines = false
 bSupportsIndependentSamplers=true
+bSupportsComputeFramework = true
 FriendlyName=LOCTEXT("FriendlyShaderPlatformName_IOSMobileSIM", "IOS Metal Simulator Mobile")
diff --git a/Engine/Plugins/Animation/DeformerGraph/DeformerGraph.uplugin b/Engine/Plugins/Animation/DeformerGraph/DeformerGraph.uplugin
--- a/Engine/Plugins/Animation/DeformerGraph/DeformerGraph.uplugin
+++ b/Engine/Plugins/Animation/DeformerGraph/DeformerGraph.uplugin
@@ -18,26 +18,18 @@
 		{
 			"Name": "OptimusSettings",
 			"Type": "Runtime",
-			"LoadingPhase": "PostConfigInit",
-			"PlatformDenyList": [
-				"IOS"
-			]
+			"LoadingPhase": "PostConfigInit"
 		},
 		{
 			"Name": "OptimusCore",
 			"Type": "Runtime",
-			"LoadingPhase": "PreDefault",
-			"PlatformDenyList": [
-				"IOS"
-			]
+			"LoadingPhase": "PreDefault"
 		},
 		{
 			"Name": "OptimusDeveloper",
 			"Type": "UncookedOnly",
-			"LoadingPhase": "PreDefault",
-			"PlatformDenyList": [
-				"IOS"
-			]
+			"LoadingPhase": "PreDefault"
+			
 		},
 		{
 			"Name": "OptimusEditor",
