Disable UdpMessaging when running on iOS. Needed as Apple requires a new local networking entitlement for local udp access and UdpMessaging isn't required for iOS game clients. It is disabled via a setting in Engine/Config/IOS/BaseIOSEngine.ini and can be reenabled with a project ini override, if still needed.

diff --git a/Engine/Config/IOS/BaseIOSEngine.ini b/Engine/Config/IOS/BaseIOSEngine.ini
--- a/Engine/Config/IOS/BaseIOSEngine.ini
+++ b/Engine/Config/IOS/BaseIOSEngine.ini
@@ -20,6 +20,10 @@ r.Substrate.ClosuresPerPixel=1
 ; Disables parallel GDME by default due to not enough threads
 r.Visibility.DynamicMeshElements.Parallel=0
 
+[/Script/UdpMessaging.UdpMessagingSettings]
+; IOS requires an additional entitlement (com.apple.developer.networking.multicast) to enable UDP Local Networking
+EnableTransport=False
+
 [Advertising]
 DefaultProviderName=IOSAdvertising
 
