iOS: Support higher quality levels for mobile bloom postprocessing effect (sync settings with Android)

diff --git a/Engine/Config/IOS/IOSScalability.ini b/Engine/Config/IOS/IOSScalability.ini
--- a/Engine/Config/IOS/IOSScalability.ini
+++ b/Engine/Config/IOS/IOSScalability.ini
@@ -105,7 +105,7 @@ r.LensFlareQuality=1
 
 [PostProcessQuality@2]
 r.Tonemapper.Quality=3
-r.BloomQuality=1
+r.BloomQuality=2
 r.DepthOfFieldQuality=0
 r.LightShaftQuality=0
 r.EyeAdaptationQuality=1
@@ -114,7 +114,7 @@ r.LensFlareQuality=2
 
 [PostProcessQuality@3]
 r.Tonemapper.Quality=3
-r.BloomQuality=1
+r.BloomQuality=3
 r.DepthOfFieldQuality=1
 r.LightShaftQuality=1
 r.EyeAdaptationQuality=1
