Remove unused iOS plugin reference

diff --git a/Engine/Config/IOS/BaseIOSEngine.ini b/Engine/Config/IOS/BaseIOSEngine.ini
--- a/Engine/Config/IOS/BaseIOSEngine.ini
+++ b/Engine/Config/IOS/BaseIOSEngine.ini
@@ -108,9 +108,6 @@ DefaultMemoryBucket_MinGB=4
 SmallerMemoryBucket_MinGB=3
 SmallestMemoryBucket_MinGB=2
 
-[Plugins]
-+EnabledPlugins=Crashlytics
-
 [Ping]
 StackSize=0
 
