Change the IOS preview options in order to mirror Android preview

#Synced-CL 52441908
#Synced-CL 52441908

diff --git a/Engine/Config/IOS/DataDrivenPlatformInfo.ini b/Engine/Config/IOS/DataDrivenPlatformInfo.ini
--- a/Engine/Config/IOS/DataDrivenPlatformInfo.ini
+++ b/Engine/Config/IOS/DataDrivenPlatformInfo.ini
@@ -32,8 +32,8 @@ ActiveIconName=LevelEditor.PreviewMode.iOS.Enabled
 InactiveIconName=LevelEditor.PreviewMode.iOS.Disabled
 MenuTooltip=NSLOCTEXT("PreviewPlatform", "PreviewMenuTooltip_iOS", "Mobile preview using iOS material quality settings.")
 IconText=NSLOCTEXT("PreviewPlatform", "PreviewIconText_iOS", "iOS")
-DeviceProfileName=IOS:iPhone6S:iPadPro6_129
-FriendlyName=NSLOCTEXT("PreviewPlatform", "PreviewFriendlyName_GenericIOS", "Generic IOS"):NSLOCTEXT("PreviewPlatform", "PreviewFriendlyName_iPhone6S", "iPhone 6S"):NSLOCTEXT("PreviewPlatform", "PreviewFriendlyName_iPadPro612.9in", "iPad Pro 6 12.9in")
+DeviceProfileName=IOS:IOS_Low:IOS_Mid:IOS_High
+FriendlyName=NSLOCTEXT("PreviewPlatform", "PreviewFriendlyName_GenericIOS", "Generic IOS"):NSLOCTEXT("PreviewPlatform", "PreviewFriendlyName_IOSLow","IOS Low"):NSLOCTEXT("PreviewPlatform", "PreviewFriendlyName_IOSMid","IOS Mid"):NSLOCTEXT("PreviewPlatform", "PreviewFriendlyName_IOSHigh","IOS High")
 
 [PreviewPlatform IOSMetalSM5]
 PlatformName=IOS
